Changed the category for the analytics plugins for better grouping

diff --git a/Engine/Plugins/Runtime/Analytics/AnalyticsBlueprintLibrary/AnalyticsBlueprintLibrary.uplugin b/Engine/Plugins/Runtime/Analytics/AnalyticsBlueprintLibrary/AnalyticsBlueprintLibrary.uplugin
--- a/Engine/Plugins/Runtime/Analytics/AnalyticsBlueprintLibrary/AnalyticsBlueprintLibrary.uplugin
+++ b/Engine/Plugins/Runtime/Analytics/AnalyticsBlueprintLibrary/AnalyticsBlueprintLibrary.uplugin
@@ -8,7 +8,7 @@
 	"CreatedByURL" : "http://epicgames.com",
 	"EngineVersion" : "4.2.0",
 	"Description" : "Blueprint Library for using an analytic event provider",
-	"Category" : "Analytics",
+	"Category" : "Analytics.BlueprintLibrary",
 	
 	"Modules" :
 	[
diff --git a/Engine/Plugins/Runtime/Analytics/Apsalar/IOSApsalar/IOSApsalar.uplugin b/Engine/Plugins/Runtime/Analytics/Apsalar/IOSApsalar/IOSApsalar.uplugin
--- a/Engine/Plugins/Runtime/Analytics/Apsalar/IOSApsalar/IOSApsalar.uplugin
+++ b/Engine/Plugins/Runtime/Analytics/Apsalar/IOSApsalar/IOSApsalar.uplugin
@@ -8,7 +8,7 @@
 	"CreatedByURL" : "http://epicgames.com",
 	"EngineVersion" : "4.2.0",
 	"Description" : "Apsalar Analytics Provider for iOS",
-	"Category" : "Analytics",
+	"Category" : "Analytics.Provider",
 	
 	"Modules" :
 	[
