Added a pure bp function for creating an analytics event attribute struct

diff --git a/Engine/Plugins/Runtime/Analytics/AnalyticsBlueprintLibrary/Source/AnalyticsBlueprintLibrary/Classes/AnalyticsBlueprintLibrary.h b/Engine/Plugins/Runtime/Analytics/AnalyticsBlueprintLibrary/Source/AnalyticsBlueprintLibrary/Classes/AnalyticsBlueprintLibrary.h
--- a/Engine/Plugins/Runtime/Analytics/AnalyticsBlueprintLibrary/Source/AnalyticsBlueprintLibrary/Classes/AnalyticsBlueprintLibrary.h
+++ b/Engine/Plugins/Runtime/Analytics/AnalyticsBlueprintLibrary/Source/AnalyticsBlueprintLibrary/Classes/AnalyticsBlueprintLibrary.h
@@ -63,6 +63,10 @@ class UAnalyticsBlueprintLibrary :
 	UFUNCTION(BlueprintCallable, Category="Analytics")
 	static void RecordCurrencyGiven(const FString& GameCurrencyType, int32 GameCurrencyAmount);
 
+	/** Builds a struct from the attribute name and value */
+	UFUNCTION(BlueprintPure, Category="Analytics")
+	static FAnalyticsEventAttr MakeEventAttribute(const FString& AttributeName, const FString& AttributeValue);
+
 	/** Gets the current session id from the analytics provider */
 	UFUNCTION(BlueprintCallable, Category="Analytics")
 	static FString GetSessionId();
diff --git a/Engine/Plugins/Runtime/Analytics/AnalyticsBlueprintLibrary/Source/AnalyticsBlueprintLibrary/Private/AnalyticsBlueprintLibrary.cpp b/Engine/Plugins/Runtime/Analytics/AnalyticsBlueprintLibrary/Source/AnalyticsBlueprintLibrary/Private/AnalyticsBlueprintLibrary.cpp
--- a/Engine/Plugins/Runtime/Analytics/AnalyticsBlueprintLibrary/Source/AnalyticsBlueprintLibrary/Private/AnalyticsBlueprintLibrary.cpp
+++ b/Engine/Plugins/Runtime/Analytics/AnalyticsBlueprintLibrary/Source/AnalyticsBlueprintLibrary/Private/AnalyticsBlueprintLibrary.cpp
@@ -219,3 +219,11 @@ void UAnalyticsBlueprintLibrary::SetUserId(const FString& UserId)
 		UE_LOG(LogAnalyticsBPLib, Warning, TEXT("SetUserId: Failed to get the default analytics provider. Double check your [Analytics] configuration in your INI"));
 	}
 }
+
+FAnalyticsEventAttr UAnalyticsBlueprintLibrary::MakeEventAttribute(const FString& AttributeName, const FString& AttributeValue)
+{
+	FAnalyticsEventAttr EventAttr;
+	EventAttr.Name = AttributeName;
+	EventAttr.Value = AttributeValue;
+	return EventAttr;
+}
