Updated the analytics plugins to point to their corresponding documentation

diff --git a/Engine/Plugins/Runtime/Analytics/AnalyticsBlueprintLibrary/AnalyticsBlueprintLibrary.uplugin b/Engine/Plugins/Runtime/Analytics/AnalyticsBlueprintLibrary/AnalyticsBlueprintLibrary.uplugin
--- a/Engine/Plugins/Runtime/Analytics/AnalyticsBlueprintLibrary/AnalyticsBlueprintLibrary.uplugin
+++ b/Engine/Plugins/Runtime/Analytics/AnalyticsBlueprintLibrary/AnalyticsBlueprintLibrary.uplugin
@@ -1,21 +1,21 @@
 {
-	"FileVersion" : 3,
-	
-	"FriendlyName" : "Analytics Blueprint Library",
-	"Version" : 1,
-	"VersionName" : "1.0",
-	"CreatedBy" : "Epic Games, Inc.",
-	"CreatedByURL" : "http://epicgames.com",
-	"EngineVersion" : "4.2.0",
-	"Description" : "Blueprint Library for using an analytic event provider",
-	"Category" : "Analytics.BlueprintLibrary",
-	
-	"Modules" :
-	[
-		{
-			"Name" : "AnalyticsBlueprintLibrary",
-			"LoadingPhase" : "PreDefault",
-			"Type" : "Runtime"
-		}
-	]
-}
\ No newline at end of file
+    "FileVersion": 3,
+
+    "FriendlyName": "Analytics Blueprint Library",
+    "Version": 1,
+    "VersionName": "1.0",
+    "CreatedBy": "Epic Games, Inc.",
+    "CreatedByURL": "http://epicgames.com",
+	"DocsURL" : "https://docs.unrealengine.com/latest/INT/Gameplay/Analytics/Blueprints/index.html",
+    "EngineVersion": "4.2.0",
+    "Description": "Blueprint Library for using an analytic event provider",
+    "Category": "Analytics.BlueprintLibrary",
+
+    "Modules": [
+        {
+            "Name": "AnalyticsBlueprintLibrary",
+            "LoadingPhase": "PreDefault",
+            "Type": "Runtime"
+        }
+    ]
+}
diff --git a/Engine/Plugins/Runtime/Analytics/Apsalar/IOSApsalar/IOSApsalar.uplugin b/Engine/Plugins/Runtime/Analytics/Apsalar/IOSApsalar/IOSApsalar.uplugin
--- a/Engine/Plugins/Runtime/Analytics/Apsalar/IOSApsalar/IOSApsalar.uplugin
+++ b/Engine/Plugins/Runtime/Analytics/Apsalar/IOSApsalar/IOSApsalar.uplugin
@@ -6,6 +6,7 @@
 	"VersionName" : "1.0",
 	"CreatedBy" : "Epic Games, Inc.",
 	"CreatedByURL" : "http://epicgames.com",
+	"DocsURL" : "https://docs.unrealengine.com/latest/INT/Gameplay/Analytics/index.html",
 	"EngineVersion" : "4.2.0",
 	"Description" : "Apsalar Analytics Provider for iOS",
 	"Category" : "Analytics.Provider",
